Fixed spurious notification when attempting to assign the Project Thumbnail (in Project Settings) to itself.

diff --git a/Engine/Source/Developer/ExternalImagePicker/Private/SExternalImagePicker.cpp b/Engine/Source/Developer/ExternalImagePicker/Private/SExternalImagePicker.cpp
--- a/Engine/Source/Developer/ExternalImagePicker/Private/SExternalImagePicker.cpp
+++ b/Engine/Source/Developer/ExternalImagePicker/Private/SExternalImagePicker.cpp
@@ -289,7 +289,8 @@ FReply SExternalImagePicker::OnPickFile()
 		{
 			check(OutFiles.Num() == 1);
 
-			if(OnExternalImagePicked.Execute(FPaths::ConvertRelativePathToFull(OutFiles[0]), TargetImagePath))
+			FString SourceImagePath = FPaths::ConvertRelativePathToFull(OutFiles[0]);
+			if (SourceImagePath != TargetImagePath && OnExternalImagePicked.Execute(SourceImagePath, TargetImagePath))
 			{
 				ApplyImage();
 			}
