Added missing include
Supplied by [REDACTED_PERSON] of The Coalition.
Fixes some non unity compile errors

diff --git a/Engine/Source/Developer/ExternalImagePicker/Public/IExternalImagePickerModule.h b/Engine/Source/Developer/ExternalImagePicker/Public/IExternalImagePickerModule.h
--- a/Engine/Source/Developer/ExternalImagePicker/Public/IExternalImagePickerModule.h
+++ b/Engine/Source/Developer/ExternalImagePicker/Public/IExternalImagePickerModule.h
@@ -6,6 +6,7 @@
 #include "Modules/ModuleInterface.h"
 #include "Modules/ModuleManager.h"
 #include "Editor/PropertyEditor/Public/PropertyHandle.h"
+#include "Framework/SlateDelegates.h"
 
 /**
  * Delegate fired when picking a new image.
