fix EngineTest failing to load its project thumbnail
fix ExternalImagePicker image loading code to use modern method
fix IcoImageWrapper was not setting BitDepth

diff --git a/Engine/Source/Developer/ExternalImagePicker/ExternalImagePicker.Build.cs b/Engine/Source/Developer/ExternalImagePicker/ExternalImagePicker.Build.cs
--- a/Engine/Source/Developer/ExternalImagePicker/ExternalImagePicker.Build.cs
+++ b/Engine/Source/Developer/ExternalImagePicker/ExternalImagePicker.Build.cs
@@ -10,11 +10,12 @@ public class ExternalImagePicker : ModuleRules
             new string[]
             {
 				"Core",
-                "Slate",
+				"Slate",
 				"SlateCore",
 				"DesktopPlatform",
 				"ImageWrapper",
-				
+				"ImageCore",
+
 				"InputCore",
 				"PropertyEditor",	// for 'reset to default'
             }
diff --git a/Engine/Source/Developer/ExternalImagePicker/Private/SExternalImagePicker.cpp b/Engine/Source/Developer/ExternalImagePicker/Private/SExternalImagePicker.cpp
--- a/Engine/Source/Developer/ExternalImagePicker/Private/SExternalImagePicker.cpp
+++ b/Engine/Source/Developer/ExternalImagePicker/Private/SExternalImagePicker.cpp
@@ -281,34 +281,24 @@ TSharedPtr< FSlateDynamicImageBrush > SExternalImagePicker::LoadImageAsBrush( co
 	if( FFileHelper::LoadFileToArray( RawFileData, *ImagePath ) )
 	{
 		IImageWrapperModule& ImageWrapperModule = FModuleManager::LoadModuleChecked<IImageWrapperModule>( FName("ImageWrapper") );
-		TSharedPtr<IImageWrapper> ImageWrappers[4] =
-		{ 
-ImageWrapperModule.CreateImageWrapper(EImageFormat::PNG),
-ImageWrapperModule.CreateImageWrapper(EImageFormat::BMP),
-ImageWrapperModule.CreateImageWrapper(EImageFormat::ICO),
-ImageWrapperModule.CreateImageWrapper(EImageFormat::ICNS),
-		};
-
-		for (auto ImageWrapper : ImageWrappers)
+
+		FImage Image;
+		if ( ImageWrapperModule.DecompressImage(RawFileData.GetData(), RawFileData.Num(), Image) )
 		{
-			if (ImageWrapper.IsValid() && ImageWrapper->SetCompressed(RawFileData.GetData(), RawFileData.Num()))
+			Image.ChangeFormat(ERawImageFormat::BGRA8,EGammaSpace::sRGB);
+
+			// have a TArray64 but API needs a TArray
+			TArray<uint8> ImageRawData32( MoveTemp(Image.RawData) );
+			if (FSlateApplication::Get().GetRenderer()->GenerateDynamicImageResource(*ImagePath, Image.SizeX, Image.SizeY, ImageRawData32))
 			{
-				TArray<uint8> RawData;
-				if (ImageWrapper->GetRaw(ERGBFormat::BGRA, 8, RawData))
-				{
-					if (FSlateApplication::Get().GetRenderer()->GenerateDynamicImageResource(*ImagePath, ImageWrapper->GetWidth(), ImageWrapper->GetHeight(), RawData))
-					{
-						Brush = MakeShareable(new FSlateDynamicImageBrush(*ImagePath, FVector2D(ImageWrapper->GetWidth(), ImageWrapper->GetHeight())));
-						bSucceeded = true;
-						break;
-					}
-				}
+				Brush = MakeShareable(new FSlateDynamicImageBrush(*ImagePath, FVector2D(Image.SizeX, Image.SizeY)));
+				bSucceeded = true;
 			}
 		}
 
 		if (!bSucceeded)
 		{
-			UE_LOG(LogSlate, Log, TEXT("Only BGRA pngs, bmps or icos are supported in by External Image Picker"));
+			UE_LOG(LogSlate, Log, TEXT("External Image Picker: DecompressImage failed"));
 			ErrorHintWidget->SetError(LOCTEXT("BadFormatHint", "Unsupported image format"));
 		}
 		else
@@ -330,7 +320,7 @@ ImageWrapperModule.CreateImageWrapper(EImageFormat::ICNS),
 	}
 	else
 	{
-	UE_LOG(LogSlate, Log, TEXT("Could not find file for image: %s"), *ImagePath);
+		UE_LOG(LogSlate, Log, TEXT("Could not find file for image: %s"), *ImagePath);
 	}
 
 	return Brush;
@@ -369,10 +359,20 @@ FReply SExternalImagePicker::OnPickFile()
 		{
 			check(OutFiles.Num() == 1);
 
+			// do not reload image if it's the same path 
 			FString SourceImagePath = FPaths::ConvertRelativePathToFull(OutFiles[0]);
-			if (SourceImagePath != TargetImagePath && OnExternalImagePicked.Execute(SourceImagePath, TargetImagePath))
+			if (SourceImagePath != TargetImagePath)
 			{
-				ApplyImageWithExtenstion(FPaths::GetExtension(SourceImagePath));
+				if ( ! OnExternalImagePicked.Execute(SourceImagePath, TargetImagePath) )
+				{
+					// can fail due to source control and other things
+					
+					UE_LOG(LogSlate, Warning, TEXT("OnExternalImagePicked.Execute failed : %s, %s"), *SourceImagePath, *TargetImagePath);
+				}
+				else
+				{
+					ApplyImageWithExtenstion(FPaths::GetExtension(SourceImagePath));
+				}
 			}
 		}
 	}
diff --git a/Engine/Source/Runtime/ImageWrapper/Private/Formats/IcoImageWrapper.cpp b/Engine/Source/Runtime/ImageWrapper/Private/Formats/IcoImageWrapper.cpp
--- a/Engine/Source/Runtime/ImageWrapper/Private/Formats/IcoImageWrapper.cpp
+++ b/Engine/Source/Runtime/ImageWrapper/Private/Formats/IcoImageWrapper.cpp
@@ -151,6 +151,7 @@ bool FIcoImageWrapper::LoadICOHeader()
 					Width = PngWrapper->GetWidth();
 					Height = PngWrapper->GetHeight();
 					Format = PngWrapper->GetFormat();
+					BitDepth = PngWrapper->GetBitDepth();
 					LargestWidth = RealWidth;
 					bFoundImage = true;
 					bIsPng = true;
@@ -165,6 +166,7 @@ bool FIcoImageWrapper::LoadICOHeader()
 					Width = BmpWrapper->GetWidth();
 					Height = BmpWrapper->GetHeight() / 2;	// ICO file spec says to divide by 2 here as height refers to combined image & mask height
 					Format = BmpWrapper->GetFormat();
+					BitDepth = BmpWrapper->GetBitDepth();
 					LargestWidth = RealWidth;
 					bFoundImage = true;
 					bIsPng = false;
